Add the LiveLinkFaceImporter plugin into UE5.  This was previously a marketplace plugin.

diff --git a/Engine/Plugins/Experimental/LiveLinkFaceImporter/LiveLinkFaceImporter.uplugin b/Engine/Plugins/Experimental/LiveLinkFaceImporter/LiveLinkFaceImporter.uplugin
new file mode 100644
--- /dev/null
+++ b/Engine/Plugins/Experimental/LiveLinkFaceImporter/LiveLinkFaceImporter.uplugin
@@ -0,0 +1,33 @@
+{
+	"FileVersion": 3,
+	"Version": 1,
+	"VersionName": "1.0",
+	"FriendlyName": "LiveLinkFaceImporter",
+	"Description": "Imports CSV recordings from the Live Link Face app.",
+	"Category": "Other",
+	"CreatedBy": "Epic Games, Inc.",
+	"CreatedByURL": "http://www.epicgames.com",
+	"DocsURL": "",
+	"MarketplaceURL": "",
+	"SupportURL": "",
+	"CanContainContent": true,
+	"IsBetaVersion": true,
+	"IsExperimentalVersion": false,
+	"Installed": false,
+	"Modules": [
+		{
+			"Name": "LiveLinkFaceImporter",
+			"Type": "Editor",
+			"LoadingPhase": "Default",
+			"PlatformAllowList": [
+				"Win64"
+			]
+		}
+	],
+	"Plugins" : [
+		{
+			"Name": "LiveLink",
+			"Enabled": true
+		},
+	]
+}
diff --git a/Engine/Plugins/Experimental/LiveLinkFaceImporter/Source/LiveLinkFaceImporter/LiveLinkFaceImporter.Build.cs b/Engine/Plugins/Experimental/LiveLinkFaceImporter/Source/LiveLinkFaceImporter/LiveLinkFaceImporter.Build.cs
new file mode 100644
--- /dev/null
+++ b/Engine/Plugins/Experimental/LiveLinkFaceImporter/Source/LiveLinkFaceImporter/LiveLinkFaceImporter.Build.cs
@@ -0,0 +1,38 @@
+// Copyright Epic Games, Inc. All Rights Reserved.
+
+using UnrealBuildTool;
+
+public class LiveLinkFaceImporter : ModuleRules
+{
+	public LiveLinkFaceImporter(ReadOnlyTargetRules Target) : base(Target)
+	{
+		PCHUsage = ModuleRules.PCHUsageMode.UseExplicitOrSharedPCHs;
+				
+		PublicDependencyModuleNames.AddRange(
+			new string[]
+			{
+				"Core",
+				"LiveLink"
+			}
+			);
+			
+		
+		PrivateDependencyModuleNames.AddRange(
+			new string[]
+			{
+				"CoreUObject",
+				"Engine",
+				"Slate",
+				"SlateCore",
+				"UnrealEd",
+				"LevelSequence",
+				"MovieScene",
+				"MovieSceneTracks",
+				"LiveLink",
+				"LiveLinkInterface",
+				"LiveLinkMovieScene"
+			}
+			);
+		
+	}
+}
diff --git a/Engine/Plugins/Experimental/LiveLinkFaceImporter/Source/LiveLinkFaceImporter/Private/LiveLinkFaceImporterFactory.cpp b/Engine/Plugins/Experimental/LiveLinkFaceImporter/Source/LiveLinkFaceImporter/Private/LiveLinkFaceImporterFactory.cpp
new file mode 100644
--- /dev/null
+++ b/Engine/Plugins/Experimental/LiveLinkFaceImporter/Source/LiveLinkFaceImporter/Private/LiveLinkFaceImporterFactory.cpp
@@ -0,0 +1,310 @@
+// Copyright Epic Games, Inc. All Rights Reserved.
+
+#include "LiveLinkFaceImporterFactory.h"
+#include "LiveLinkFaceImporterLog.h"
+#include "LevelSequence.h"
+#include "Misc/FileHelper.h"
+#include "Misc/QualifiedFrameTime.h"
+#include "Misc/Timecode.h"
+#include "MovieScene.h"
+#include "MovieScene/MovieSceneLiveLinkTrack.h"
+#include "MovieScene/MovieSceneLiveLinkSection.h"
+#include "Roles/LiveLinkBasicRole.h"
+#include "LiveLinkPresetTypes.h"
+
+
+//////////////////////////////////////////////////////////////////////////
+namespace
+{
+	const int32 TIMECODE_COLUMN_INDEX = 0;
+	const int32 BLENDSHAPECOUNT_COLUMN_INDEX = 1;
+
+	const FString TIMECODE_COLUMN_NAME = TEXT("Timecode");
+	const FString BLENDSHAPECOUNT_COLUMN_NAME = TEXT("BlendShapeCount");
+}
+
+
+//////////////////////////////////////////////////////////////////////////
+// ULiveLinkFaceImporterFactory
+
+ULiveLinkFaceImporterFactory::ULiveLinkFaceImporterFactory(const FObjectInitializer& ObjectInitializer)
+	: Super(ObjectInitializer)
+{
+	bCreateNew = false;
+	SupportedClass = ULevelSequence::StaticClass();
+
+	bEditorImport = true;
+	bText = true;
+
+	Formats.Add(TEXT("csv;Live Link Face recording"));
+}
+
+FText ULiveLinkFaceImporterFactory::GetToolTip() const
+{
+	return NSLOCTEXT("LiveLinkFaceImporter", "LiveLinkFaceImporterFactoryDescription", "CSV files saved by the Live Link Face iOS app.");
+}
+
+bool ULiveLinkFaceImporterFactory::LoadCSV(const FString& FileContent, TArray<FString>& KeyArray, TArray<FString>& LineArray, FString & OutLogMessage)
+{
+	if (FileContent.ParseIntoArrayLines(LineArray) > 1)
+	{
+		// At least 2 lines in the CSV are needed : One which is the column names (keys) and at least one line
+		// which contains values.
+		if (FileContent.ParseIntoArrayLines(LineArray) < 2)
+		{
+			OutLogMessage = TEXT("The input CSV file contained no data.");
+			return false;
+		}
+
+		// The first line (the column names / keys) must have at least 3 fields : timecode, shape count, and 
+		// at least one value column. otherwise there is no data to import.
+		if (LineArray[0].ParseIntoArray(KeyArray, TEXT(",")) < 3)
+		{
+			OutLogMessage = TEXT("The input CSV file didn't include more than 2 columns.");
+			return false;
+		}
+		LineArray.RemoveAt(0);
+
+		// validate that the first two entires in the KeyArray are Timecode & BlendShapeCount
+		if (KeyArray[TIMECODE_COLUMN_INDEX] != TIMECODE_COLUMN_NAME)
+		{
+			OutLogMessage = FString::Format(TEXT("The input CSV file was poorly formatted : the first column should be \"Timecode\"."), { *TIMECODE_COLUMN_NAME });
+			return false;
+		}
+		// validate that the first two entires in the KeyArray are Timecode & BlendShapeCount
+		if (KeyArray[BLENDSHAPECOUNT_COLUMN_INDEX] != BLENDSHAPECOUNT_COLUMN_NAME)
+		{
+			OutLogMessage = FString::Format(TEXT("The input CSV file was poorly formatted : the second column should be \"%s\"."), { *BLENDSHAPECOUNT_COLUMN_NAME });
+			return false;
+		}
+
+	}
+	return true;
+}
+
+FString ULiveLinkFaceImporterFactory::CreateSubjectString(const FString& InString)
+{
+	// Try to determine the subject name : The filename is <Slate>_<Take>_<SubjectName>[_cal|_raw|_neutral].csv.
+	FString SubjectString = InString;
+
+	// Remove the [_cal|_raw|_neutral] suffix if it exists
+	SubjectString.RemoveFromEnd("_cal");
+	SubjectString.RemoveFromEnd("_raw");
+	SubjectString.RemoveFromEnd("_neutral");
+
+	{
+		TArray<FString> TokenArray;
+		if (SubjectString.ParseIntoArray(TokenArray, TEXT("_"), false) > 3)
+		{
+			// Remove the first token (the slate)
+			TokenArray.RemoveAt(0);
+
+			// Remove each token until (but including) a token that is a number : we assume it's the take number. The
+			// remaining tokens will form the subject name.
+			while (TokenArray.Num())
+			{
+				bool bIsNumeric = TokenArray[0].IsNumeric();
+				TokenArray.RemoveAt(0);
+
+				if (bIsNumeric)
+					break;
+			}
+
+			if (TokenArray.Num())
+			{
+				SubjectString = FString::Join(TokenArray, TEXT("_"));
+			}
+		}
+	}
+
+	// If we've erased everything due to an unexpected filename, just use the InString again.
+	return SubjectString.IsEmpty() ? InString : SubjectString;
+
+}
+
+bool ULiveLinkFaceImporterFactory::FactoryCanImport(const FString& Filename)
+{
+	// If the file has more than 2 lines and the first line is tokenized by commas *and* has 
+	// the first two columns named correctly, assume we can import it.
+	FString FileContent;
+	if (FFileHelper::LoadFileToString(/*out*/ FileContent, *Filename))
+	{
+		TArray<FString> LineArray;
+		TArray<FString> KeyArray;
+		FString OutLogMessage;
+
+		if (LoadCSV(FileContent, KeyArray, LineArray, OutLogMessage))
+		{
+			return true;
+		}
+		else
+		{
+			UE_LOG(LogLiveLinkFaceImporter, VeryVerbose, TEXT("%s"), *OutLogMessage);
+			return false;
+		}
+	}
+
+	UE_LOG(LogLiveLinkFaceImporter, Error, TEXT("Unable to load the file '%s'"), *Filename);
+	return false;
+}
+
+UObject* ULiveLinkFaceImporterFactory::FactoryCreateText(UClass* InClass, UObject* InParent, FName InName, EObjectFlags Flags, UObject* Context, const TCHAR* Type, const TCHAR*& Buffer, const TCHAR* BufferEnd, FFeedbackContext* Warn)
+{
+	const FString FileContent(BufferEnd - Buffer, Buffer);
+	TArray<FString> LineArray;
+	TArray<FString> KeyArray;
+	FString OutLogMessage;
+
+	// Load the CSV into Keys (the first row tokenized) and Lines (the rest of the file)
+	if (LoadCSV(FileContent, KeyArray, LineArray, OutLogMessage) == false)
+	{
+		UE_LOG(LogLiveLinkFaceImporter, Error, TEXT("%s"), *OutLogMessage);
+		return nullptr;
+	}
+
+	UE_LOG(LogLiveLinkFaceImporter, Log, TEXT("Importing CSV file."));
+
+	// Create the subject name from the InName (the filename)
+	FString SubjectString = CreateSubjectString(InName.ToString());
+
+	UE_LOG(LogLiveLinkFaceImporter, Verbose, TEXT("Using subject name '%s'."), *SubjectString);
+
+	FName SubjectName(SubjectString);
+	UClass* SubjectRole = ULiveLinkBasicRole::StaticClass();
+
+	// create the new sequence
+	ULevelSequence* NewLevelSequence = NewObject<ULevelSequence>(InParent, InName, Flags | RF_Transactional);
+	NewLevelSequence->Initialize();
+
+	// save the tick + display rate as we will need them to convert from the CSV's timecode
+	const FFrameRate TickResolution = NewLevelSequence->GetMovieScene()->GetTickResolution();
+	const FFrameRate DisplayRate = NewLevelSequence->GetMovieScene()->GetDisplayRate();
+
+	// Add the master LL track
+	UMovieSceneLiveLinkTrack* LiveLinkTrack = NewLevelSequence->GetMovieScene()->AddMasterTrack<UMovieSceneLiveLinkTrack>();
+	LiveLinkTrack->SetTrackRole(SubjectRole);
+	LiveLinkTrack->SetPropertyNameAndPath(SubjectName, SubjectName.ToString());
+
+	// Create and add a section to the LL track
+	UMovieSceneLiveLinkSection* MovieSceneSection = Cast<UMovieSceneLiveLinkSection>(LiveLinkTrack->CreateNewSection());
+
+	MovieSceneSection->SetIsActive(false);
+	LiveLinkTrack->AddSection(*MovieSceneSection);
+
+	FLiveLinkSubjectPreset SubjectPreset;
+	SubjectPreset.Key.Source.Invalidate();
+	SubjectPreset.Key.SubjectName = SubjectName;
+	SubjectPreset.Role = SubjectRole;
+	SubjectPreset.bEnabled = true;
+
+	// Add all the property names from the first row of the CSV (skipping Timecode & BlendshapeCount)
+	FLiveLinkBaseStaticData BaseStaticData;
+	for (const FString& Key : KeyArray)
+	{
+		if ((Key != TIMECODE_COLUMN_NAME) && (Key != BLENDSHAPECOUNT_COLUMN_NAME))
+		{
+			BaseStaticData.PropertyNames.Add(*Key);
+		}
+	}
+
+	// initialize the movie scene section with the given properties from the CSV
+	TSharedPtr<FLiveLinkStaticDataStruct> StaticData = MakeShared<FLiveLinkStaticDataStruct>();
+	StaticData->InitializeWith(&BaseStaticData);
+	MovieSceneSection->Initialize(SubjectPreset, StaticData);
+	MovieSceneSection->CreateChannelProxy();
+
+	int32 LineCount = 1;
+	int32 FrameCount = 0;
+
+	bool IsFirstFrame = true;
+	TArray<FString> ValueArray;
+	for (const FString& Line : LineArray)
+	{
+		// the number of values in each line must match the keys/header in the first line
+		int32 ValueCount = Line.ParseIntoArray(ValueArray, TEXT(","));
+		if (ValueCount != KeyArray.Num())
+		{
+			UE_LOG(LogLiveLinkFaceImporter, Error, TEXT("Line %d did not contain the correct number of columns (found %d, expected %d). Aborting!"), LineCount, ValueCount, KeyArray.Num());
+			break;
+		}
+
+		// now build the frame data for this line
+
+		FLiveLinkBaseFrameData FrameData;
+		double SampleTime = 0.0;
+
+		for (int32 Index = 0; Index < ValueArray.Num(); ++Index)
+		{
+			if (Index == BLENDSHAPECOUNT_COLUMN_INDEX)
+			{
+				continue;
+			}
+			else if (Index == TIMECODE_COLUMN_INDEX)
+			{
+				// Tokenize the timecode : we require 4 tokens
+				TArray<FString> TimeCodeTokens;
+				if (ValueArray[Index].ParseIntoArray(TimeCodeTokens, TEXT(":")) == 4)
+				{
+					int32 Hours = FCString::Atoi(*TimeCodeTokens[0]);
+					int32 Minutes = FCString::Atoi(*TimeCodeTokens[1]);
+					int32 Seconds = FCString::Atoi(*TimeCodeTokens[2]);
+					double Frames = FCString::Atof(*TimeCodeTokens[3]);
+
+					SampleTime = double(Hours * 60 * 60 + Minutes * 60 + Seconds) + Frames / 60.0;
+					FTimecode Timecode = FTimecode(SampleTime, FFrameRate(60,1), false);
+
+					// use this timecode as the source if we are on the first frame
+					if (IsFirstFrame)
+					{
+						MovieSceneSection->TimecodeSource = FMovieSceneTimecodeSource(Timecode);
+						IsFirstFrame = false;
+					}
+				}
+				else
+				{
+					UE_LOG(LogLiveLinkFaceImporter, Error, TEXT("Line %d contained a malformed timecode '%s' (Should be of format hh:mm:ss:frame.zzz). Skipping!"), LineCount);
+					break;
+				}
+			}
+			else
+			{
+				// add a keyframe for this property
+				FrameData.PropertyValues.Add(FCString::Atof(*ValueArray[Index]));
+			}
+
+		}
+
+		// create the frame data struct and add it to the movie scene section
+		FLiveLinkFrameDataStruct Frame;
+		Frame.InitializeWith(&FrameData);
+
+		MovieSceneSection->RecordFrame((SampleTime * TickResolution).FloorToFrame(), MoveTemp(Frame));
+
+		// clear the value array for the next frame
+		ValueArray.Empty(ValueArray.Num());
+
+		// these are for logging purposes
+		FrameCount++;
+		LineCount++;
+	}
+
+	// finish up the sequence + set the section/playback ranges
+
+	FKeyDataOptimizationParams Params;
+	Params.bAutoSetInterpolation = true;
+
+	MovieSceneSection->FinalizeSection(false, Params);
+
+	TOptional<TRange<FFrameNumber>> DefaultSectionLength = MovieSceneSection->GetAutoSizeRange();
+	if (DefaultSectionLength.IsSet())
+	{
+		MovieSceneSection->SetRange(DefaultSectionLength.GetValue());
+		NewLevelSequence->GetMovieScene()->SetPlaybackRange(DefaultSectionLength.GetValue());
+	}
+	MovieSceneSection->SetIsActive(true);
+
+	// all done
+	UE_LOG(LogLiveLinkFaceImporter, Log, TEXT("Successfully imported %d frames."), FrameCount);
+
+	return NewLevelSequence;
+}
diff --git a/Engine/Plugins/Experimental/LiveLinkFaceImporter/Source/LiveLinkFaceImporter/Private/LiveLinkFaceImporterFactory.h b/Engine/Plugins/Experimental/LiveLinkFaceImporter/Source/LiveLinkFaceImporter/Private/LiveLinkFaceImporterFactory.h
new file mode 100644
--- /dev/null
+++ b/Engine/Plugins/Experimental/LiveLinkFaceImporter/Source/LiveLinkFaceImporter/Private/LiveLinkFaceImporterFactory.h
@@ -0,0 +1,26 @@
+// Copyright Epic Games, Inc. All Rights Reserved.
+
+#pragma once
+
+#include "CoreMinimal.h"
+#include "UObject/ObjectMacros.h"
+#include "Factories/Factory.h"
+#include "LiveLinkFaceImporterFactory.generated.h"
+
+// Imports a CSV file created by the Live Link Face iOS app when recordin facial animation.
+UCLASS()
+class ULiveLinkFaceImporterFactory : public UFactory
+{
+	GENERATED_UCLASS_BODY()
+		
+	// UFactory interface
+	virtual FText GetToolTip() const override;
+	virtual bool FactoryCanImport(const FString& Filename) override;
+	virtual UObject* FactoryCreateText(UClass* InClass, UObject* InParent, FName InName, EObjectFlags Flags, UObject* Context, const TCHAR* Type, const TCHAR*& Buffer, const TCHAR* BufferEnd, FFeedbackContext* Warn) override;
+	// End of UFactory interface
+
+protected:
+	bool LoadCSV(const FString& FileContent, TArray<FString>& KeyArray, TArray<FString>& LineArray, FString& OutLogMessage);
+	FString CreateSubjectString(const FString& InName);
+};
+
diff --git a/Engine/Plugins/Experimental/LiveLinkFaceImporter/Source/LiveLinkFaceImporter/Private/LiveLinkFaceImporterLog.h b/Engine/Plugins/Experimental/LiveLinkFaceImporter/Source/LiveLinkFaceImporter/Private/LiveLinkFaceImporterLog.h
new file mode 100644
--- /dev/null
+++ b/Engine/Plugins/Experimental/LiveLinkFaceImporter/Source/LiveLinkFaceImporter/Private/LiveLinkFaceImporterLog.h
@@ -0,0 +1,8 @@
+// Copyright Epic Games, Inc. All Rights Reserved.
+
+#pragma once
+
+#include "CoreMinimal.h"
+
+DECLARE_LOG_CATEGORY_EXTERN(LogLiveLinkFaceImporter, Verbose, All);
+
diff --git a/Engine/Plugins/Experimental/LiveLinkFaceImporter/Source/LiveLinkFaceImporter/Private/LiveLinkFaceImporterModule.cpp b/Engine/Plugins/Experimental/LiveLinkFaceImporter/Source/LiveLinkFaceImporter/Private/LiveLinkFaceImporterModule.cpp
new file mode 100644
--- /dev/null
+++ b/Engine/Plugins/Experimental/LiveLinkFaceImporter/Source/LiveLinkFaceImporter/Private/LiveLinkFaceImporterModule.cpp
@@ -0,0 +1,33 @@
+// Copyright Epic Games, Inc. All Rights Reserved.
+
+#include "CoreMinimal.h"
+#include "LiveLinkFaceImporterLog.h"
+#include "Modules/ModuleManager.h"
+#include "Internationalization/Internationalization.h"
+
+#define LOCTEXT_NAMESPACE "LiveLinkFaceImporter"
+
+DEFINE_LOG_CATEGORY(LogLiveLinkFaceImporter);
+
+//////////////////////////////////////////////////////////////////////////
+// FLiveLinkFaceImporterModule
+
+class FLiveLinkFaceImporterModule : public FDefaultModuleImpl
+{
+public:
+	virtual void StartupModule() override
+	{
+	}
+
+	virtual void ShutdownModule() override
+	{
+	}
+};
+
+//////////////////////////////////////////////////////////////////////////
+
+IMPLEMENT_MODULE(FLiveLinkFaceImporterModule, LiveLinkFaceImporter);
+
+//////////////////////////////////////////////////////////////////////////
+
+#undef LOCTEXT_NAMESPACE
