Add 30FPS support to LiveLinkFaceImporter and remove platform restrictions

This change attempts to infer the framerate of the data contained within the CSV file as either 30 or 60 FPS. This is done by making a pass of timecode values and determining which frame rate setting was used when the data was recorded.

diff --git a/Engine/Plugins/Experimental/LiveLinkFaceImporter/LiveLinkFaceImporter.uplugin b/Engine/Plugins/Experimental/LiveLinkFaceImporter/LiveLinkFaceImporter.uplugin
--- a/Engine/Plugins/Experimental/LiveLinkFaceImporter/LiveLinkFaceImporter.uplugin
+++ b/Engine/Plugins/Experimental/LiveLinkFaceImporter/LiveLinkFaceImporter.uplugin
@@ -18,16 +18,13 @@
 		{
 			"Name": "LiveLinkFaceImporter",
 			"Type": "Editor",
-			"LoadingPhase": "Default",
-			"PlatformAllowList": [
-				"Win64"
-			]
+			"LoadingPhase": "Default"
 		}
 	],
 	"Plugins" : [
 		{
 			"Name": "LiveLink",
 			"Enabled": true
-		},
+		}
 	]
 }
diff --git a/Engine/Plugins/Experimental/LiveLinkFaceImporter/Source/LiveLinkFaceImporter/Private/LiveLinkFaceImporterFactory.cpp b/Engine/Plugins/Experimental/LiveLinkFaceImporter/Source/LiveLinkFaceImporter/Private/LiveLinkFaceImporterFactory.cpp
--- a/Engine/Plugins/Experimental/LiveLinkFaceImporter/Source/LiveLinkFaceImporter/Private/LiveLinkFaceImporterFactory.cpp
+++ b/Engine/Plugins/Experimental/LiveLinkFaceImporter/Source/LiveLinkFaceImporter/Private/LiveLinkFaceImporterFactory.cpp
@@ -65,13 +65,13 @@ bool ULiveLinkFaceImporterFactory::LoadCSV(const FString& FileContent, TArray<FS
 		}
 		LineArray.RemoveAt(0);
 
-		// validate that the first two entires in the KeyArray are Timecode & BlendShapeCount
+		// validate that the first two entries in the KeyArray are Timecode & BlendShapeCount
 		if (KeyArray[TIMECODE_COLUMN_INDEX] != TIMECODE_COLUMN_NAME)
 		{
 			OutLogMessage = FString::Format(TEXT("The input CSV file was poorly formatted : the first column should be \"Timecode\"."), { *TIMECODE_COLUMN_NAME });
 			return false;
 		}
-		// validate that the first two entires in the KeyArray are Timecode & BlendShapeCount
+		// validate that the first two entries in the KeyArray are Timecode & BlendShapeCount
 		if (KeyArray[BLENDSHAPECOUNT_COLUMN_INDEX] != BLENDSHAPECOUNT_COLUMN_NAME)
 		{
 			OutLogMessage = FString::Format(TEXT("The input CSV file was poorly formatted : the second column should be \"%s\"."), { *BLENDSHAPECOUNT_COLUMN_NAME });
@@ -122,6 +122,73 @@ FString ULiveLinkFaceImporterFactory::CreateSubjectString(const FString& InStrin
 
 }
 
+/**
+ * Tokenizes a time code string and returns the hours, minutes, seconds and frames as discrete values.
+ *
+ * @param	InTimecodeString A string representation of a timecode (hh:mm:ss:frame.zzz)
+ * @param	OutHours	Hours output
+ * @param	OutMinutes	Minutes output
+ * @param	OutSeconds	Seconds output
+ * @param	OutFrames	Fractional frames output 
+ *
+ * @return	Whether or not the timecode string could be tokenized correctly.
+ */
+bool ULiveLinkFaceImporterFactory::ParseTimecode(const FString& InTimecodeString, int32& OutHours,
+	int32& OutMinutes, int32& OutSeconds, double& OutFrames)
+{
+	if (TArray<FString> TimeCodeTokens; InTimecodeString.ParseIntoArray(TimeCodeTokens, TEXT(":")) == 4)
+	{
+		OutHours = FCString::Atoi(*TimeCodeTokens[0]);
+		OutMinutes = FCString::Atoi(*TimeCodeTokens[1]);
+		OutSeconds = FCString::Atoi(*TimeCodeTokens[2]);
+		OutFrames = FCString::Atof(*TimeCodeTokens[3]);
+		return true;
+	}
+	
+	return false;
+}
+
+/**
+ * @brief Infers the frame rate of the CSV data as either 30 or 60 FPS.
+ * 
+ * There is no explicit metadata available within the CSV output file generated by the Live Link Face application that
+ * describes the frame rate of the captured data.
+ * 
+ * However we know that when recording at 30 FPS the LLF app will not output timecodes with fractional frame values
+ * greater than 30. With that in mind we can infer when a file was recorded at a target of 30 or 60 FPS.
+ * @param InLineValuesArray The values contained within the CSV file. 
+ * @param OutInferredFrameRate The inferred frame rate value to populate.
+ * @return Whether or not the operation was a success. 
+ */
+bool ULiveLinkFaceImporterFactory::InferFrameRate(const TArray<TArray<FString>>& InLineValuesArray, int8& OutInferredFrameRate)
+{
+	constexpr double Minimum60FPSFrameNumber = 30.0;
+	int8 InferredFrameRate = 30;
+	for (const TArray<FString>& ValueArray : InLineValuesArray)
+	{
+		if (TIMECODE_COLUMN_INDEX < ValueArray.Num())
+		{
+			int32 Hours;
+			int32 Minutes;
+			int32 Seconds;
+			double FrameNumber;
+			
+			if (!ParseTimecode(ValueArray[TIMECODE_COLUMN_INDEX], Hours, Minutes, Seconds, FrameNumber))
+			{
+				return false;
+			}
+
+			if (FrameNumber > Minimum60FPSFrameNumber)
+			{
+				InferredFrameRate = 60;
+			}
+		} 
+	}
+
+	OutInferredFrameRate = InferredFrameRate;
+	return true;
+}
+
 bool ULiveLinkFaceImporterFactory::FactoryCanImport(const FString& Filename)
 {
 	// If the file has more than 2 lines and the first line is tokenized by commas *and* has 
@@ -178,8 +245,7 @@ UObject* ULiveLinkFaceImporterFactory::FactoryCreateText(UClass* InClass, UObjec
 
 	// save the tick + display rate as we will need them to convert from the CSV's timecode
 	const FFrameRate TickResolution = NewLevelSequence->GetMovieScene()->GetTickResolution();
-	const FFrameRate DisplayRate = NewLevelSequence->GetMovieScene()->GetDisplayRate();
-
+	
 	// Add the LL track
 	UMovieSceneLiveLinkTrack* LiveLinkTrack = NewLevelSequence->GetMovieScene()->AddTrack<UMovieSceneLiveLinkTrack>();
 	LiveLinkTrack->SetTrackRole(SubjectRole);
@@ -213,23 +279,41 @@ UObject* ULiveLinkFaceImporterFactory::FactoryCreateText(UClass* InClass, UObjec
 	MovieSceneSection->Initialize(SubjectPreset, StaticData);
 	MovieSceneSection->CreateChannelProxy();
 
-	int32 LineCount = 1;
-	int32 FrameCount = 0;
+	// Tokenize each line within the CSV file
+	TArray<TArray<FString>> LineValuesArray;
+	LineValuesArray.Reserve(LineArray.Num());
+	for (const FString& Line : LineArray)
+	{
+		TArray<FString> ValueArray;
+		Line.ParseIntoArray(ValueArray, TEXT(","));
+		LineValuesArray.Add(ValueArray);
+	}
 
+	// Infer the frame rate of the file contents in the absence of any explicit metadata.
+	int8 InferredFrameRate;
+	if (!InferFrameRate(LineValuesArray, InferredFrameRate))
+	{
+		UE_LOG(LogLiveLinkFaceImporter, Warning, TEXT("Failed to infer frame rate from CSV file contents. Assuming 60 FPS."))
+		InferredFrameRate = 60;
+	}
+
+	// Set the intiial display rate of the movie scene to that inferred from the CSV file contents.
+	NewLevelSequence->GetMovieScene()->SetDisplayRate(FFrameRate(InferredFrameRate, 1));
+	
+	int32 LineNumber = 1;
+	int32 FrameCount = 0;
 	bool IsFirstFrame = true;
-	TArray<FString> ValueArray;
-	for (const FString& Line : LineArray)
+	for (const TArray<FString>& ValueArray : LineValuesArray)
 	{
 		// the number of values in each line must match the keys/header in the first line
-		int32 ValueCount = Line.ParseIntoArray(ValueArray, TEXT(","));
+		int32 ValueCount = ValueArray.Num();
 		if (ValueCount != KeyArray.Num())
 		{
-			UE_LOG(LogLiveLinkFaceImporter, Error, TEXT("Line %d did not contain the correct number of columns (found %d, expected %d). Aborting!"), LineCount, ValueCount, KeyArray.Num());
+			UE_LOG(LogLiveLinkFaceImporter, Error, TEXT("Line %d did not contain the correct number of columns (found %d, expected %d). Aborting!"), LineNumber, ValueCount, KeyArray.Num());
 			break;
 		}
-
+		
 		// now build the frame data for this line
-
 		FLiveLinkBaseFrameData FrameData;
 		double SampleTime = 0.0;
 
@@ -241,17 +325,14 @@ UObject* ULiveLinkFaceImporterFactory::FactoryCreateText(UClass* InClass, UObjec
 			}
 			else if (Index == TIMECODE_COLUMN_INDEX)
 			{
-				// Tokenize the timecode : we require 4 tokens
-				TArray<FString> TimeCodeTokens;
-				if (ValueArray[Index].ParseIntoArray(TimeCodeTokens, TEXT(":")) == 4)
+				int32 Hours;
+				int32 Minutes;
+				int32 Seconds;
+				double Frames;
+				if (ParseTimecode(ValueArray[Index], Hours, Minutes, Seconds, Frames))
 				{
-					int32 Hours = FCString::Atoi(*TimeCodeTokens[0]);
-					int32 Minutes = FCString::Atoi(*TimeCodeTokens[1]);
-					int32 Seconds = FCString::Atoi(*TimeCodeTokens[2]);
-					double Frames = FCString::Atof(*TimeCodeTokens[3]);
-
-					SampleTime = double(Hours * 60 * 60 + Minutes * 60 + Seconds) + Frames / 60.0;
-					FTimecode Timecode = FTimecode(SampleTime, FFrameRate(60,1), false);
+					SampleTime = double(Hours * 60 * 60 + Minutes * 60 + Seconds) + Frames / static_cast<double>(InferredFrameRate);
+					FTimecode Timecode = FTimecode(SampleTime, FFrameRate(InferredFrameRate,1), false);
 
 					// use this timecode as the source if we are on the first frame
 					if (IsFirstFrame)
@@ -262,7 +343,7 @@ UObject* ULiveLinkFaceImporterFactory::FactoryCreateText(UClass* InClass, UObjec
 				}
 				else
 				{
-					UE_LOG(LogLiveLinkFaceImporter, Error, TEXT("Line %d contained a malformed timecode (Should be of format hh:mm:ss:frame.zzz). Skipping!"), LineCount);
+					UE_LOG(LogLiveLinkFaceImporter, Error, TEXT("Line %d contained a malformed timecode (Should be of format hh:mm:ss:frame.zzz). Skipping!"), LineNumber);
 					break;
 				}
 			}
@@ -271,7 +352,6 @@ UObject* ULiveLinkFaceImporterFactory::FactoryCreateText(UClass* InClass, UObjec
 				// add a keyframe for this property
 				FrameData.PropertyValues.Add(FCString::Atof(*ValueArray[Index]));
 			}
-
 		}
 
 		// create the frame data struct and add it to the movie scene section
@@ -280,12 +360,9 @@ UObject* ULiveLinkFaceImporterFactory::FactoryCreateText(UClass* InClass, UObjec
 
 		MovieSceneSection->RecordFrame((SampleTime * TickResolution).FloorToFrame(), MoveTemp(Frame));
 
-		// clear the value array for the next frame
-		ValueArray.Empty(ValueArray.Num());
-
 		// these are for logging purposes
 		FrameCount++;
-		LineCount++;
+		LineNumber++;
 	}
 
 	// finish up the sequence + set the section/playback ranges
diff --git a/Engine/Plugins/Experimental/LiveLinkFaceImporter/Source/LiveLinkFaceImporter/Private/LiveLinkFaceImporterFactory.h b/Engine/Plugins/Experimental/LiveLinkFaceImporter/Source/LiveLinkFaceImporter/Private/LiveLinkFaceImporterFactory.h
--- a/Engine/Plugins/Experimental/LiveLinkFaceImporter/Source/LiveLinkFaceImporter/Private/LiveLinkFaceImporterFactory.h
+++ b/Engine/Plugins/Experimental/LiveLinkFaceImporter/Source/LiveLinkFaceImporter/Private/LiveLinkFaceImporterFactory.h
@@ -7,7 +7,7 @@
 #include "Factories/Factory.h"
 #include "LiveLinkFaceImporterFactory.generated.h"
 
-// Imports a CSV file created by the Live Link Face iOS app when recordin facial animation.
+// Imports a CSV file created by the Live Link Face iOS app when recording facial animation.
 UCLASS()
 class ULiveLinkFaceImporterFactory : public UFactory
 {
@@ -22,5 +22,7 @@ class ULiveLinkFaceImporterFactory : public UFactory
 protected:
 	bool LoadCSV(const FString& FileContent, TArray<FString>& KeyArray, TArray<FString>& LineArray, FString& OutLogMessage);
 	FString CreateSubjectString(const FString& InName);
+	bool ParseTimecode(const FString& InTimecodeString, int32& OutHours, int32& OutMinutes, int32& OutSeconds, double& OutFrames);
+	bool InferFrameRate(const TArray<TArray<FString>>& InLineValuesArray, int8& OutInferredFrameRate);
 };
 
